[INTEGRATE] Change 2445863 by [REDACTED_PERSON][REDACTED_HANDLE] on 2015/02/13 21:14:22

	UEBP-34:  Final fixup for Leap integration

diff --git a/Engine/Plugins/Runtime/LeapMotionController/Source/LeapMotionController/Private/LeapMotionBoneActor.cpp b/Engine/Plugins/Runtime/LeapMotionController/Source/LeapMotionController/Private/LeapMotionBoneActor.cpp
--- a/Engine/Plugins/Runtime/LeapMotionController/Source/LeapMotionController/Private/LeapMotionBoneActor.cpp
+++ b/Engine/Plugins/Runtime/LeapMotionController/Source/LeapMotionController/Private/LeapMotionBoneActor.cpp
@@ -74,7 +74,7 @@ void ALeapMotionBoneActor::Init(ELeapBone LeapBone, float Scale, float Width, fl
 		Length = 0.0f;
 		break;
 	default: 
-		SetMeshForArm(); 
+		SetMeshForFinger();
 		break;
 	}
 
