Eliminate GWorld in leap motion plugin

diff --git a/Engine/Plugins/Runtime/LeapMotionController/Source/LeapMotionController/Private/LeapMotionControllerComponent.cpp b/Engine/Plugins/Runtime/LeapMotionController/Source/LeapMotionController/Private/LeapMotionControllerComponent.cpp
--- a/Engine/Plugins/Runtime/LeapMotionController/Source/LeapMotionController/Private/LeapMotionControllerComponent.cpp
+++ b/Engine/Plugins/Runtime/LeapMotionController/Source/LeapMotionController/Private/LeapMotionControllerComponent.cpp
@@ -175,7 +175,7 @@ void ULeapMotionControllerComponent::OnHandAddedImpl(int32 HandId)
 		SpawnRotation = (SpawnRotation.Quaternion() * Roll.Quaternion() * ForwardTilt.Quaternion()).Rotator();
 	}
 
-	ALeapMotionHandActor* handActor = GWorld->SpawnActor<ALeapMotionHandActor>(HandBlueprint ? HandBlueprint : ALeapMotionHandActor::StaticClass(), SpawnLocation, SpawnRotation, SpawnParams);
+	ALeapMotionHandActor* handActor = GetWorld()->SpawnActor<ALeapMotionHandActor>(HandBlueprint ? HandBlueprint : ALeapMotionHandActor::StaticClass(), SpawnLocation, SpawnRotation, SpawnParams);
 
 	if (handActor)
 	{
diff --git a/Engine/Plugins/Runtime/LeapMotionController/Source/LeapMotionController/Private/LeapMotionHandActor.cpp b/Engine/Plugins/Runtime/LeapMotionController/Source/LeapMotionController/Private/LeapMotionHandActor.cpp
--- a/Engine/Plugins/Runtime/LeapMotionController/Source/LeapMotionController/Private/LeapMotionHandActor.cpp
+++ b/Engine/Plugins/Runtime/LeapMotionController/Source/LeapMotionController/Private/LeapMotionHandActor.cpp
@@ -120,7 +120,7 @@ void ALeapMotionHandActor::CreateBones(const TSubclassOf<class ALeapMotionBoneAc
 				Position = RefQuat * Position * CombinedScale + GetRootComponent()->GetComponentLocation();
 				Orientation = (RefQuat * Orientation.Quaternion()).Rotator();
 
-				ALeapMotionBoneActor* BoneActor = GWorld->SpawnActor<ALeapMotionBoneActor>(BoneBlueprint ? BoneBlueprint : ALeapMotionBoneActor::StaticClass(), Position, Orientation, SpawnParams);
+				ALeapMotionBoneActor* BoneActor = GetWorld()->SpawnActor<ALeapMotionBoneActor>(BoneBlueprint ? BoneBlueprint : ALeapMotionBoneActor::StaticClass(), Position, Orientation, SpawnParams);
 				if (BoneActor) 
 				{
 					BoneActors.Add(BoneActor);
