#ue4 - Perf: Avoid pointless Quat->Rotator->Quat conversions.

diff --git a/Engine/Plugins/Runtime/LeapMotionController/Source/LeapMotionController/Private/LeapMotionHandActor.cpp b/Engine/Plugins/Runtime/LeapMotionController/Source/LeapMotionController/Private/LeapMotionHandActor.cpp
--- a/Engine/Plugins/Runtime/LeapMotionController/Source/LeapMotionController/Private/LeapMotionHandActor.cpp
+++ b/Engine/Plugins/Runtime/LeapMotionController/Source/LeapMotionController/Private/LeapMotionHandActor.cpp
@@ -116,7 +116,7 @@ void ALeapMotionHandActor::CreateBones(const TSubclassOf<class ALeapMotionBoneAc
 			Success &= Device->GetBoneWidthAndLength(HandId, LeapBone, Width, Length);
 			if (Success)
 			{
-				FQuat RefQuat = GetRootComponent()->GetComponentRotation().Quaternion();
+				const FQuat RefQuat = GetRootComponent()->GetComponentQuat();
 				Position = RefQuat * Position * CombinedScale + GetRootComponent()->GetComponentLocation();
 				Orientation = (RefQuat * Orientation.Quaternion()).Rotator();
 
@@ -165,7 +165,7 @@ void ALeapMotionHandActor::UpdateBones(float DeltaSeconds)
 			if (Success)
 			{
 				// Offset target position & rotation by the SpawnReference actor's transform
-				FQuat RefQuat = GetRootComponent()->GetComponentRotation().Quaternion();
+				const FQuat RefQuat = GetRootComponent()->GetComponentQuat();
 				TargetPosition = RefQuat * TargetPosition * CombinedScale + GetRootComponent()->GetComponentLocation();
 				TargetOrientation = (RefQuat * TargetOrientation.Quaternion()).Rotator();
 
diff --git a/Engine/Source/Runtime/Engine/Private/Particles/ParticleTrail2EmitterInstance.cpp b/Engine/Source/Runtime/Engine/Private/Particles/ParticleTrail2EmitterInstance.cpp
--- a/Engine/Source/Runtime/Engine/Private/Particles/ParticleTrail2EmitterInstance.cpp
+++ b/Engine/Source/Runtime/Engine/Private/Particles/ParticleTrail2EmitterInstance.cpp
@@ -2333,8 +2333,7 @@ bool FParticleRibbonEmitterInstance::ResolveSourcePoint(int32 InTrailIdx,
 				OutPosition += SourceOffsetValue;
 			}
 		}
-		FRotator TempRotator = Component->GetComponentRotation();
-		OutRotation = FQuat(TempRotator);
+		OutRotation = Component->GetComponentQuat();
 		OutTangent = Component->PartSysVelocity;
 		OutTangentStrength = OutTangent.SizeSquared();
 		OutUp = Component->ComponentToWorld.GetScaledAxis( EAxis::Z );
