Set up runtime dependency for LeapMotion plugin.

diff --git a/Engine/Plugins/Runtime/LeapMotionController/Source/LeapMotionController/LeapMotionController.Build.cs b/Engine/Plugins/Runtime/LeapMotionController/Source/LeapMotionController/LeapMotionController.Build.cs
--- a/Engine/Plugins/Runtime/LeapMotionController/Source/LeapMotionController/LeapMotionController.Build.cs
+++ b/Engine/Plugins/Runtime/LeapMotionController/Source/LeapMotionController/LeapMotionController.Build.cs
@@ -11,6 +11,8 @@ namespace UnrealBuildTool.Rules
 			PublicDependencyModuleNames.AddRange( new string[] { "Core", "CoreUObject", "Engine", "InputCore" } ); 
 
             PublicDependencyModuleNames.AddRange(new string[] { "Leap" }); 
+
+			RuntimeDependencies.Add(new RuntimeDependency("$(EngineDir)/Binaries/ThirdParty/Leap/Leap.dll"));
 		}
 	}
 }
\ No newline at end of file
