#leapmotion Pulling latest doc and crash updates from Leap Motion

#github 1033

diff --git a/Engine/Plugins/Runtime/LeapMotionController/Source/LeapMotionController/Private/LeapMotionControllerComponent.cpp b/Engine/Plugins/Runtime/LeapMotionController/Source/LeapMotionController/Private/LeapMotionControllerComponent.cpp
--- a/Engine/Plugins/Runtime/LeapMotionController/Source/LeapMotionController/Private/LeapMotionControllerComponent.cpp
+++ b/Engine/Plugins/Runtime/LeapMotionController/Source/LeapMotionController/Private/LeapMotionControllerComponent.cpp
@@ -29,9 +29,7 @@ ULeapMotionControllerComponent::ULeapMotionControllerComponent(const class FObje
 
 	Scale = 5.0f;
 
-	IsInHmdMode = false;
-	UseHmdMode(IsInHmdMode);
-
+	bHmdMode = false;
 	bShowCollider = true;
 	bShowMesh = true;
 }
@@ -79,12 +77,12 @@ ALeapMotionHandActor* ULeapMotionControllerComponent::GetOldestLeftOrRightHandAc
 
 void ULeapMotionControllerComponent::UseHmdMode(bool EnableOrDisable)
 {
-	IsInHmdMode = EnableOrDisable;
+	bHmdMode = EnableOrDisable;
 	
 	FLeapMotionDevice* Device = FLeapMotionControllerPlugin::GetLeapDeviceSafe();
 	if (Device && Device->IsConnected())
 	{
-		Device->SetHmdPolicy(IsInHmdMode);
+		Device->SetHmdPolicy(bHmdMode);
 	}
 }
 
@@ -96,11 +94,17 @@ void ULeapMotionControllerComponent::TickComponent(float DeltaTime, enum ELevelT
 	UdpateHandsPositions(DeltaTime);
 }
 
+void ULeapMotionControllerComponent::PostInitProperties()
+{
+	Super::PostInitProperties();
+	UseHmdMode(bHmdMode);
+}
+
 #if WITH_EDITOR
 void ULeapMotionControllerComponent::PostEditChangeProperty(FPropertyChangedEvent& PropertyChangedEvent)
 {
 	Super::PostEditChangeProperty(PropertyChangedEvent);
-	UseHmdMode(IsInHmdMode);
+	UseHmdMode(bHmdMode);
 }
 #endif
 
@@ -168,7 +172,7 @@ void ULeapMotionControllerComponent::OnHandAddedImpl(int32 HandId)
 	SpawnParams.Owner = GetOwner();
 	SpawnParams.Instigator = GetOwner()->GetInstigator();
 
-	if (IsInHmdMode)
+	if (bHmdMode)
 	{
 		FRotator ForwardTilt(-90, 0, 0);
 		FRotator Roll(0, 0, 180);
diff --git a/Engine/Plugins/Runtime/LeapMotionController/Source/LeapMotionController/Private/LeapMotionHandActor.cpp b/Engine/Plugins/Runtime/LeapMotionController/Source/LeapMotionController/Private/LeapMotionHandActor.cpp
--- a/Engine/Plugins/Runtime/LeapMotionController/Source/LeapMotionController/Private/LeapMotionHandActor.cpp
+++ b/Engine/Plugins/Runtime/LeapMotionController/Source/LeapMotionController/Private/LeapMotionHandActor.cpp
@@ -29,7 +29,12 @@ ALeapMotionBoneActor* ALeapMotionHandActor::GetBoneActor(ELeapBone LeapBone) con
 {
 	if (BoneActors.Num())
 	{
-		return BoneActors[(int)LeapBone - 1 + bShowArm];
+		bool HasArm = (BoneActors.Num() == int(ELeapBone::Finger4Tip) + 1);
+		int Index = (int)LeapBone - 1 + HasArm;
+		if (0 <= Index && Index < BoneActors.Num())
+		{
+			return BoneActors[Index];
+		}
 	}
 	return nullptr;
 }
diff --git a/Engine/Plugins/Runtime/LeapMotionController/Source/LeapMotionController/Public/LeapMotionBoneActor.h b/Engine/Plugins/Runtime/LeapMotionController/Source/LeapMotionController/Public/LeapMotionBoneActor.h
--- a/Engine/Plugins/Runtime/LeapMotionController/Source/LeapMotionController/Public/LeapMotionBoneActor.h
+++ b/Engine/Plugins/Runtime/LeapMotionController/Source/LeapMotionController/Public/LeapMotionBoneActor.h
@@ -11,63 +11,91 @@ class ALeapMotionHandActor;
 class ULeapMotionControllerComponent;
 
 /**
- * This actor is dynamically created by LeapMotionHandActor. It represents a single hand bone tracked by Leap Motion Controller.
+ * Represents a single part of a hand tracked by the Leap Motion Controller.
+ *
+ * Bone Actors are created automatically by a LeapMotionHandActor instance. You can
+ * extended this class to replace or augment the default behavior.
+ *
+ * Bone actors contain the position and orientation information. These attributes 
+ * are relative to the parent LeapMotionHandActor, but note that by default a
+ * Hand actor has the same position and orientation as its parent
+ * LeapMotionControllerComponent object.
  */
 UCLASS(ClassGroup = LeapMotion, BlueprintType, Blueprintable)
 class LEAPMOTIONCONTROLLER_API ALeapMotionBoneActor : public AActor
 {
 	GENERATED_UCLASS_BODY()
 
+	/** The name of this bone. */
 	UPROPERTY(VisibleAnywhere, BlueprintReadOnly, Category = LeapMotion)
 	ELeapBone LeapBoneType;
 
-	/** Visualize hand bone's colliding shape in game */
+	/** Whether to visualize hand bone's colliding shape in game. */
 	UPROPERTY(EditAnywhere, BlueprintReadWrite, Category = LeapMotion)
 	bool bShowCollider;
 
-	/** Visualize hand bone's mesh in game */
+	/** Whether to render the hand bone's mesh in game. */
 	UPROPERTY(EditAnywhere, BlueprintReadWrite, Category = LeapMotion)
 	bool bShowMesh;
 
-	/** Primitive component used for collisions with the scene objects */
+	/** The Unreal Primitive component used for collisions with the scene objects. */
 	UPROPERTY(EditAnywhere, BlueprintReadWrite, Category = LeapMotion)
 	class UPrimitiveComponent* PrimitiveComponent;
 
-	/** Static mesh for visualizing the hand bone */
+	/** The Static mesh for rendering the hand bone. */
 	UPROPERTY(EditAnywhere, BlueprintReadWrite, Category = LeapMotion)
 	class UStaticMeshComponent* StaticMeshComponent;
 
-	/** Gets the controller component that created this bone & it's parenting hand. */
+	/** 
+	 * Gets the controller component that created this bone and its parenting hand.
+	 * @returns The parent LeapMotionControllerComponent.
+	 */
 	UFUNCTION(BlueprintPure, Category = LeapMotion)
 	ULeapMotionControllerComponent* GetParentingControllerComponent() const;
 
-	/** Gets the hand actor this bone belongs to. */
+	/** 
+	 * Gets the hand actor this bone belongs to. 
+	 * @returns the parent LeapMotionHandActor instance.
+	 */
 	UFUNCTION(BlueprintPure, Category = LeapMotion)
 	ALeapMotionHandActor* GetHandActor() const;
 
-	/** Sizes the mesh & colliding shap eof this hand bone */
+	/** 
+	 * Called to initialize the size and display properties of a bone based on 
+	 * the current hand settings.
+	 *
+	 * Subclasses of LeapMotionBoneActor can implement this function to override 
+	 * the default behavior.
+	 *
+	 * @param LeapBone The name of the current bone.
+	 * @param Scale The relative scale for this bone.
+	 * @param Width The width of this bone.
+	 * @param Length The length of this bone.
+	 * @param ShowCollider Whether to render the colliders.
+	 * @param ShowMesh Whether to render the mesh for this bone.
+	 */
 	UFUNCTION(BlueprintCallable, Category = LeapMotion)
 	virtual void Init(ELeapBone LeapBone, float Scale, float Width, float Length, bool ShowCollider, bool ShowMesh);
 
-	/** Updates visibility of hand bone's mesh & colliding shape */
+	/** Updates visibility of hand bone's mesh and colliding shape. */
 	void UpdateVisibility();
 
 #if WITH_EDITOR
-	/** Track changes to property applied in the editor */
+	/** Track changes to properties applied in the editor. */
 	virtual void PostEditChangeProperty(FPropertyChangedEvent& PropertyChangedEvent) override;
 #endif
 
 protected:
 
-	/** Uses the default finger mesh */
+	/** Uses the default finger mesh. */
 	UFUNCTION(BlueprintCallable, Category = LeapMotion)
 	virtual void SetMeshForFinger();
 
-	/** Uses the default palm mesh */
+	/** Uses the default palm mesh. */
 	UFUNCTION(BlueprintCallable, Category = LeapMotion)
 	virtual void SetMeshForPalm();
 
-	/** Uses the default arm mesh */
+	/** Uses the default arm mesh. */
 	UFUNCTION(BlueprintCallable, Category = LeapMotion)
 	virtual void SetMeshForArm();
 
diff --git a/Engine/Plugins/Runtime/LeapMotionController/Source/LeapMotionController/Public/LeapMotionControllerActor.h b/Engine/Plugins/Runtime/LeapMotionController/Source/LeapMotionController/Public/LeapMotionControllerActor.h
--- a/Engine/Plugins/Runtime/LeapMotionController/Source/LeapMotionController/Public/LeapMotionControllerActor.h
+++ b/Engine/Plugins/Runtime/LeapMotionController/Source/LeapMotionController/Public/LeapMotionControllerActor.h
@@ -9,15 +9,16 @@
 
 
 /**
- *  Place this actor in the world to automatically create hand actors as tracked by Leap Motion 
- *  controller. You can also attach this Controller to a character or a camera.
+ * An Unreal Actor instance that allows you to drop a set of 3D hands anywhere in a level.
+ *
+ * Contains a LeapMotionControllerComponent instance which does all of the work.
  */
 UCLASS(ClassGroup = LeapMotion, BlueprintType, Blueprintable)
 class LEAPMOTIONCONTROLLER_API ALeapMotionControllerActor : public AActor
 {
 	GENERATED_UCLASS_BODY()
 
-	/** Scale of hands, relatively to real-world size. This ignores actor's transform's scale */
+	/** An Unreal Component instance. */
 	UPROPERTY(VisibleAnywhere, BlueprintReadOnly, Category = LeapMotion)
 	class ULeapMotionControllerComponent* LeapMotionControllerComponent;
 
diff --git a/Engine/Plugins/Runtime/LeapMotionController/Source/LeapMotionController/Public/LeapMotionControllerComponent.h b/Engine/Plugins/Runtime/LeapMotionController/Source/LeapMotionController/Public/LeapMotionControllerComponent.h
--- a/Engine/Plugins/Runtime/LeapMotionController/Source/LeapMotionController/Public/LeapMotionControllerComponent.h
+++ b/Engine/Plugins/Runtime/LeapMotionController/Source/LeapMotionController/Public/LeapMotionControllerComponent.h
@@ -21,25 +21,28 @@ class LEAPMOTIONCONTROLLER_API ULeapMotionControllerComponent : public USceneCom
 {
 	GENERATED_UCLASS_BODY()
 
-	/** Scale of hands, relatively to real-world size. This ignores actor's transform's scale */
+	/** The scale of hands, relative to real-world size. This ignores actor's transform's scale */
 	UPROPERTY(EditAnywhere, BlueprintReadWrite, Category = LeapMotion)
 	float Scale;
 
-	/** Visualize hand's colliding shape in game */
+	/** Whether to render the hand's colliding shape in game */
 	UPROPERTY(EditAnywhere, BlueprintReadWrite, Category = LeapMotion)
 	bool bShowCollider;
 
-	/** Visualize hand's mesh in game */
+	/** Whether to render the hand's mesh in game */
 	UPROPERTY(EditAnywhere, BlueprintReadWrite, Category = LeapMotion)
 	bool bShowMesh;
 
-	/** Show hand's arm */
+	/** Whether to render the hand's arm */
 	UPROPERTY(EditAnywhere, BlueprintReadWrite, Category = LeapMotion)
 	bool bShowArm;
 
-	/** Run the controller as when attached to an HMD */
+	/** 
+	 * Whether to transform the tracking data as appropriate for a Leap Motion device
+	 * attached to an HMD. 
+	 */
 	UPROPERTY(EditAnywhere, BlueprintReadWrite, Category = LeapMotion)
-	bool IsInHmdMode;
+	bool bHmdMode;
 
 	/** Optional blueprint used to spawn each hand. When none is specified ALeapMotionHandActor is created. */
 	UPROPERTY(EditAnywhere, BlueprintReadWrite, Category = LeapMotion)
@@ -50,100 +53,113 @@ class LEAPMOTIONCONTROLLER_API ULeapMotionControllerComponent : public USceneCom
 	TSubclassOf<class ALeapMotionBoneActor> BoneBlueprint;
 
 	/** 
-	 * Gets all active hands ids 
-	 * @param OutHandIds	Array of ids of existing hand actors.
+	 * Gets all active hands ids.
+	 * @param OutHandIds Output array which is filled with the ids of existing hands.
 	 */
 	UFUNCTION(BlueprintPure, Category = LeapMotion)
 	void GetAllHandIds(TArray<int32>& OutHandIds) const;
 
 	/** 
-	 * Gets all active hands actors 
-	 * @param OutHandActors		Array of existing hand actors.
+	 * Gets all active hands actors.
+	 * @param OutHandActors	Output array which is filled with existing hand actors.
 	 */
 	UFUNCTION(BlueprintPure, Category = LeapMotion)
 	void GetAllHandActors(TArray<ALeapMotionHandActor*>& OutHandActors) const;
 
 	/** 
-	 * Returns hand actor 
-	 * @param HandId	Which hand to query
-	 * @return			Actor pointer if hand is found, nullptr otherwise.
+	 * Gets the hand actor instance for the specified hand ID.
+	 * @param HandId The Leap Motion ID for the hand of interest.
+	 * @returns	A pointer to the Actor object if the hand ID is active, nullptr otherwise.
 	 */
 	UFUNCTION(BlueprintPure, Category = LeapMotion)
 	ALeapMotionHandActor* GetHandActor(int32 HandId) const;
 
 	/**
 	 * Returns the oldest left- or right-hand actor, if one exists, nullptr otherwise. 
+	 *
+	 * If more than one left or right hand is being tracked, this function returns 
+	 * the one that has been tracked the longest.
 	 * @param LeapSide	Are we looking for left or right-sided hands?
 	 * @return			Pointer to oldest hand actor if one is found, nullptr otherwise.
 	 */
 	UFUNCTION(BlueprintPure, Category = LeapMotion)
 	ALeapMotionHandActor* GetOldestLeftOrRightHandActor(ELeapSide LeapSide) const;
 
-	/** Delegate triggered when a new hand is spotted by the Leap Motion device */
+	/** Delegate triggered when a new hand is spotted by the Leap Motion device. */
 	UPROPERTY(BlueprintAssignable, Category = LeapMotion)
 	FLMHandAddedDelegate OnHandAdded;
 
-	/** Delegate triggered when a hand lost by the Leap Motion device */
+	/** Delegate triggered when a hand lost by the Leap Motion device. */
 	UPROPERTY(BlueprintAssignable, Category = LeapMotion)
 	FLMHandRemovedDelegate OnHandRemoved;
 
-	/** Delegate triggered once per frame update for each hand */
+	/** Delegate triggered once per frame update for each hand. */
 	UPROPERTY(BlueprintAssignable, Category = LeapMotion)
 	FLMHandUpdatedDelegate OnHandUpdated;
 
-	/** Sets the controller as when attached to an HMD, or in the standard desktop mode */
+	/** Sets the controller as when attached to an HMD, or in the standard desktop mode 
+	 * Puts the Leap Motion controller and this component into (or out of) HMD mode.
+	 * This sets the Leap Motion software to use the Optimize for HMD policy.
+	 * It also changes the transform of Leap Motion to Unreal coordinates so that the
+	 * hands are projected in front of the component position rather than above it.
+	 *
+	 * @param EnableOrDisable Set to true to enable HMD mode.
+	 */
 	UFUNCTION(BlueprintCallable, Category = LeapMotion)
 	void UseHmdMode(bool EnableOrDisable);
 
 	/**
-	 * Updates all hands. 
+	 * Updates all hand actors of this component.
 	 * @param DeltaTime - The time since the last tick.
-	 * @param TickType - The kind of tick this is
-	 * @param ThisTickFunction - Tick function that caused this to run
+	 * @param TickType - The kind of tick this is.
+	 * @param ThisTickFunction - Tick function that caused this to run.
 	 */
 	virtual void TickComponent(float DeltaTime, enum ELevelTick TickType, FActorComponentTickFunction *ThisTickFunction);
 
-	/** Finger model mesh used by the default hand implementation */
+	/** Finger model mesh used by the default hand implementation. */
 	UPROPERTY(EditAnywhere, BlueprintReadWrite, Category = LeapMotion)
 	UStaticMesh* FingerMesh;
 
-	/** Palm model mesh used by the default hand implementation */
+	/** Palm model mesh used by the default hand implementation. */
 	UPROPERTY(EditAnywhere, BlueprintReadWrite, Category = LeapMotion)
 	UStaticMesh* PalmMesh;
 
-	/** Arm model mesh used by the default hand implementation */
+	/** Arm model mesh used by the default hand implementation. */
 	UPROPERTY(EditAnywhere, BlueprintReadWrite, Category = LeapMotion)
 	UStaticMesh* ArmMesh;
 
-	/** Material used by the default hand implementation */
+	/** Material used by the default hand implementation. */
 	UPROPERTY(EditAnywhere, BlueprintReadWrite, Category = LeapMotion)
 	UMaterialInterface* Material;
 
+	/** Called after the C++ constructor and after the properties have been initialized */
+	virtual void PostInitProperties() override;
+
 #if WITH_EDITOR
-	/** Track changes to property applied in the editor */
+	/** Track changes to property applied in the editor. */
 	virtual void PostEditChangeProperty(FPropertyChangedEvent& PropertyChangedEvent) override;
 #endif
 
 protected:
 
-	/** Created & destroys hand actors & triggers corresponding delegates */
+	/** Creates and destroys hand actors and triggers the corresponding delegates. */
 	void AddAndRemoveHands();
 
-	/** Updates the hands & calls the update delegate */
+	/** Updates the hands and calls the update delegate. */
 	void UdpateHandsPositions(float DeltaSeconds);
 
-	/** Creates a hand */
+	/** Creates a hand. */
 	virtual void OnHandAddedImpl(int32 HandId);
 	
-	/** Destroys a hand */
+	/** Destroys a hand. */
 	virtual void OnHandRemovedImpl(int32 HandId);
 
-	/** Updates a hand */
+	/** Updates a hand. */
 	virtual void OnHandUpdatedImpl(int32 HandId, float DeltaSeconds);
 
-	/** List of ids in the last frame */
+	/** List of Leap Motion hand ids in the last frame. */
 	TArray<int32> LastFrameHandIds;
 
-	/** All actors for active hand ids */
+	/** All actors for active hand ids. */
 	TMap<int32, class ALeapMotionHandActor*> HandActors;
 };
diff --git a/Engine/Plugins/Runtime/LeapMotionController/Source/LeapMotionController/Public/LeapMotionControllerPlugin.h b/Engine/Plugins/Runtime/LeapMotionController/Source/LeapMotionController/Public/LeapMotionControllerPlugin.h
--- a/Engine/Plugins/Runtime/LeapMotionController/Source/LeapMotionController/Public/LeapMotionControllerPlugin.h
+++ b/Engine/Plugins/Runtime/LeapMotionController/Source/LeapMotionController/Public/LeapMotionControllerPlugin.h
@@ -7,20 +7,20 @@
 DECLARE_LOG_CATEGORY_EXTERN(LogLeapMotion, All, All);
 
 /**
- * The public interface to this module
+ * The public interface to the Leap Motion Controller plugin module.
  */
 class FLeapMotionControllerPlugin : public IModuleInterface, public IModularFeature
 {
 
 public:
 
-    // IModuleInterface implementation.
-    virtual void StartupModule();
-    virtual void ShutdownModule();
+	// IModuleInterface implementation.
+	virtual void StartupModule();
+	virtual void ShutdownModule();
 
 	/**
-	 * Singleton-like access to this module's interface.  This is just for convenience!
-	 * Beware of calling this during the shutdown phase, though.  Your module might have been unloaded already.
+	 * Singleton-like access to this module's interface.
+	 * Beware of calling this during the shutdown phase.  The module might have been unloaded already.
 	 *
 	 * @return Returns singleton instance, loading the module on demand if needed
 	 */
@@ -39,29 +39,29 @@ public:
 		return FModuleManager::Get().IsModuleLoaded( "LeapMotionController" );
 	}
 
-    FORCEINLINE TSharedPtr<class FLeapMotionDevice> GetLeapDevice() const
-    {
-        return LeapDevice;
-    }
+	FORCEINLINE TSharedPtr<class FLeapMotionDevice> GetLeapDevice() const
+	{
+		return LeapDevice;
+	}
 
-    /**
-     * Simple helper function to get the device currently active.
-	 * @return	Pointer to the LeapMotionDevice, or nullptr if Device is not available.
-     */
-    static FLeapMotionDevice* GetLeapDeviceSafe()
-    {
+	/**
+	 * Simple helper function to get the device currently active.
+	 * @returns	Pointer to the LeapMotionDevice, or nullptr if Device is not available.
+	 */
+	static FLeapMotionDevice* GetLeapDeviceSafe()
+	{
 #if WITH_EDITOR
-        FLeapMotionDevice* LeapDevice = FLeapMotionControllerPlugin::IsAvailable() ? FLeapMotionControllerPlugin::Get().GetLeapDevice().Get() : nullptr;
+		FLeapMotionDevice* LeapDevice = FLeapMotionControllerPlugin::IsAvailable() ? FLeapMotionControllerPlugin::Get().GetLeapDevice().Get() : nullptr;
 #else
 		FLeapMotionDevice* LeapDevice = FLeapMotionControllerPlugin::Get().GetLeapDevice().Get();
 #endif
-        return LeapDevice;
-    }
+		return LeapDevice;
+	}
 
 protected:
-    /**
-    * Reference to the actual Leap Device, grabbed through the GetLeapDevice() interface, and created and destroyed in Startup/ShutdownModule
-    */
-    TSharedPtr<class FLeapMotionDevice> LeapDevice;
+	/**
+	 * Reference to the actual Leap Device, grabbed through the GetLeapDevice() interface, and created and destroyed in Startup/ShutdownModule
+	 */
+	TSharedPtr<class FLeapMotionDevice> LeapDevice;
 };
 
diff --git a/Engine/Plugins/Runtime/LeapMotionController/Source/LeapMotionController/Public/LeapMotionDevice.h b/Engine/Plugins/Runtime/LeapMotionController/Source/LeapMotionController/Public/LeapMotionDevice.h
--- a/Engine/Plugins/Runtime/LeapMotionController/Source/LeapMotionController/Public/LeapMotionDevice.h
+++ b/Engine/Plugins/Runtime/LeapMotionController/Source/LeapMotionController/Public/LeapMotionDevice.h
@@ -28,81 +28,123 @@ public:
 	void ShutdownDevice();
 
 	/**
-	 * Checks whether or not the device is connected and ready to be polled for data
+	 * Checks whether or not the device is connected and ready to be polled for data.
 	 * @return			True if controller is connected, false otherwise.
 	 */
 	bool IsConnected() const;
 
 	/**
-	 * Returns list of IDs of all hands seen by the device
-	 * @param OutAllHandIds		Array of all visible hand IDs.
-	 * @return					True if the device is connected 
+	 * Returns list of IDs of all hands tracked by the device
+	 * @param OutAllHandIds		An output array set to contain the ids of all tracked hands.
+	 * @return					True, if the device is connected
 	 */
 	bool GetAllHandsIds(TArray<int32>& OutAllHandIds) const;
 
 	/** 
-	 * Returns the id of the oldest left or right hand, if one exists. 
-	 * @param LeapSide	Query for left or right hands
-	 * @param OutHandId	The Id of the oldest left/right hand, or -1 otherwise.
-	 * @return			True if hand is found, false otherwise.
+	 * Returns the oldest left- or right-hand actor, if one exists, nullptr otherwise. 
+	 *
+	 * If more than one left or right hand is being tracked, this function returns 
+	 * the one that has been tracked the longest.
+	 * @param LeapSide	Look for a left or a right hand.
+	 * @param OutHandId	An integer set to the Leap Motion id of the oldest left 
+	 *					or right hand. If no hand is found, this is set to -1.
+	 * @returns			True if a hand of the specified type exists, false otherwise.
 	 */
 	bool GetOldestLeftOrRightHandId(ELeapSide LeapSide, int32& OutHandId) const;
 
 	/**
-	 * Tells if this is a left or a right hand
-	 * @param HandId			Which hand to query
-	 * @param OutIsLeft			Tells if this is a left hand
-	 * @return					True if the hand was detected , false otherwise
+	 * Reports whether the specified hand us identified as a left or a right hand.
+	 * @param HandId			The id of the hand of interest.
+	 * @param OutIsLeft			A bool set to true if the specified hand is a left hand.
+	 * @returns			        True if a hand of the specified type exists, false otherwise.
 	 */
 	bool GetIsHandLeft(int32 HandId, bool& OutIsLeft) const;
 
 	/**
-	 * Gets the position and orientation of the specified hand from the device
-	 * @param HandId			Which hand to query
+	 * Gets the position and orientation of the specified hand from the device.
+	 *
+	 * Data returned by this function is relative to the Unreal Engine origin rather than
+	 * a particular LeapMotionControllerComponent or LeapMotionHandActor instance.
+	 * Get the properties of a LeapMotionBoneActor instance to get the data
+	 * relative to that actor's parent hand actor and controller component. In this
+	 * case, use the standard Unreal Actor location and rotation properties.
+	 *
+	 * @param HandId			The id of the hand of interest.
 	 * @param LeapBone			Which hand bone to get position & orientation for
-	 * @param OutPosition		(out) The relative position of the hand, or zero if the hand wasn't detected this frame
-	 * @param OutOrientation	(out) The relative orientation of the hand, or zero if the hand wasn't detected this frame
-	 * @return					True if the hand was detected, false otherwise
+	 * @param OutPosition		An FVector set to the relative position of the hand, 
+	 *							or zero if the hand wasn't detected this frame.
+	 * @param OutOrientation	An FRotator set to the relative orientation of the hand, 
+	 *							or zero if the hand wasn't detected this frame.
+	 * @return					True if a hand with the specified id exists, false otherwise
 	 */
 	bool GetBonePostionAndOrientation(int32 HandId, ELeapBone LeapBone, FVector& OutPosition, FRotator& OutOrientation) const;
 
 	/**
-	 * Gets the width and length of the specified bone of a hand from the device
-	 * @param HandId			Which hand to query
-	 * @param LeapBone			Which hand bone to get position & orientation for
-	 * @param OutWidth			Bone's width (in centimeters)
-	 * @param OutWidth			Bone's length (in centimeters)
+	 * Gets the width and length of the specified bone of a hand from the device.
+	 * @param HandId			The id of the hand of interest.
+	 * @param LeapBone			The bone of interest.
+	 * @param OutWidth			A float set to the bone's width (in centimeters).
+	 * @param OutLength			A float set to the bone's length (in centimeters).
 	 * @return					True if the hand was detected and bone was identified, false otherwise
 	 */
 	bool GetBoneWidthAndLength(int32 HandId, ELeapBone LeapBone, float& OutWidth, float& OutLength) const;
 
 	/**
-	 * Queries grabbing strength of the hand
-	 * @param HandId			Which hand to query
-	 * @param OutGrabStrength	Hand's grabbing strength as reported by Leap API
-	 * @return					True if the hand was detected, false otherwise
+	 * The grab strength rating of the specified hand.
+	 *
+	 * Grab strength is a rating of how the hand's posture resembles a fist.
+	 * A strength of 0 is close to an open, flat hand; a strength of 1 is close 
+	 * to a fist.
+	 * @param HandId			The id of the hand of interest.
+	 * @param OutGrabStrength	A float set to the hand's grabbing strength as reported by Leap API.
+	 * @returns					True, if the hand with the specified id exists, false otherwise.
 	 */
 	bool GetGrabStrength(int32 HandId, float& OutGrabStrength) const;
 
 	/**
-	 * Queries pinching strength of the hand
-	 * @param HandId			Which hand to query
-	 * @param OutPinchStrength	Hand's pinch strength as reported by Leap API
-	 * @return					True if the hand was detected, false otherwise
+	 * The pinch strength rating of the specified hand.
+	 *
+	 * Pinch strength is a rating of whether the hand is in a pinching posture.
+	 * Pinch strength starts at 0 and increases to 1 as any finger tip approaches
+	 * the tip of the thumb. Note that pinch and grab strength are not independent.
+	 * A grabbing hand will generally have a non-zero pinch strength as well.
+	 *
+	 * @param HandId			The id of the hand of interest.
+	 * @param OutPinchStrength	A float set to the hand's pinch strength as reported by Leap API
+	 * @returns					True, if the hand with the specified id exists, false otherwise.
 	 */
 	bool GetPinchStrength(int32 HandId, float& OutPinchStrength) const;
 
 	/**
-	 * Set LeapController policy that improves tracking for HMD-mounted controller
-	 * @param UseHmdPolicy		True for HMD-mounted mode, false otherwise
-	 * @return					True if the device is connected 
+	 * Enables or disables the Leap Motion Controller policy that improves tracking for an
+	 * HMD-mounted controller.
+	 * Note that calling this function does not change the transforms of any 
+	 * LeapMotionControllerComponent instances that may exist.
+	 *
+	 * @param UseHmdPolicy		True to enable for HMD-mounted mode, false to disable.
+	 * @returns					True if the device is connected.
 	 */
 	bool SetHmdPolicy(bool UseHmdPolicy);
 
-	/** Gets reference frame */
+	/** 
+	 * Gets the cached reference frame. Call SetReferenceFrameOncePerTick() before
+	 * the first access to this variable during a tick.
+	 *
+	 * Data in a Leap::Frame object has not been transformed into the Unreal Engine
+	 * frame of reference or unit system.
+	 */
 	const Leap::Frame Frame() const;
 
-	/** Caches the current Leap::Frame as a reference. Call once per unreal frame tick, and before any other calls to LeapMotionDevice on that frame tick. */
+	/** 
+	 * Caches the current Leap::Frame as a reference. Call once per Unreal frame 
+	 * tick, and before any other calls to LeapMotionDevice on that frame tick. 
+	 * The Leap Motion software produces frames of data at a rate independent of
+	 * the Unreal Engine. Cache the frame so that subsequent calls to get data 
+	 * from the device use the same data.
+	 *
+	 * Calling the function a second time in the same tick has no affect, but it
+	 * does need to be called at least once.
+	 */
 	void SetReferenceFrameOncePerTick();
 
 protected:
diff --git a/Engine/Plugins/Runtime/LeapMotionController/Source/LeapMotionController/Public/LeapMotionFunctionLibrary.h b/Engine/Plugins/Runtime/LeapMotionController/Source/LeapMotionController/Public/LeapMotionFunctionLibrary.h
--- a/Engine/Plugins/Runtime/LeapMotionController/Source/LeapMotionController/Public/LeapMotionFunctionLibrary.h
+++ b/Engine/Plugins/Runtime/LeapMotionController/Source/LeapMotionController/Public/LeapMotionFunctionLibrary.h
@@ -22,79 +22,101 @@ class LEAPMOTIONCONTROLLER_API ULeapMotionFunctionLibrary : public UBlueprintFun
 
 	/** 
 	 * Checks whether a Leap Motion controller is connected.
-	 * @return			True if controller is connected, false otherwise.
+	 * @returns			True if controller is connected, false otherwise.
 	 */
 	UFUNCTION(BlueprintPure, Category = LeapMotion)
 	static bool IsConnected();
 
 	/**
-	 * Returns list of IDs of all hands seen by the device
-	 * @param OutAllHandIds		Array of all visible hand IDs.
-	 * @param UseReferenceFrame	Use saved reference frame, rather then the volatile newest one from the device.
-	 * @return					True if the device is connected 
+	 * Returns list of IDs of all hands tracked by the device.
+	 * @param OutAllHandIds		Output array which is filled with all the tracked hand IDs.
+	 * @returns					True, if the device is connected.
 	 */
 	UFUNCTION(BlueprintPure, Category = LeapMotion)
 	static bool GetAllHandsIds(TArray<int32>& OutAllHandIds);
 
 	/** 
-	 * Returns the id of the oldest left or right hand, if one exists. 
-	 * @param LeapSide	Query for left or right hands
-	 * @param OutHandId	The Id of the oldest left/right hand, or -1 otherwise.
-	 * @return			True if hand is found, false otherwise.
+	 * Returns the oldest left- or right-hand actor, if one exists, nullptr otherwise. 
+	 *
+	 * If more than one left or right hand is being tracked, this function returns 
+	 * the one that has been tracked the longest.
+	 * @param LeapSide	Look for a left or a right hand.
+	 * @param OutHandId	An integer set to the Leap Motion id of the oldest left
+	 *					or right hand. If no hand is found, this is set to -1.
+	 * @returns			True if a hand of the specified type exists, false otherwise.
 	 */
 	UFUNCTION(BlueprintPure, Category = LeapMotion)
 	static bool GetOldestLeftOrRightHandId(ELeapSide LeapSide, int32& OutHandId);
 
 	/**
-	 * Gets the position and orientation of the specified hand from the device
-	 * @param HandId			Which hand to query
-	 * @param LeapBone			Which hand bone to get position & orientation for
-	 * @param OutPosition		(out) The relative position of the hand, or zero if the hand wasn't detected this frame
-	 * @param OutOrientation	(out) The relative orientation of the hand, or zero if the hand wasn't detected this frame
-	 * @param UseReferenceFrame	Use saved reference frame, rather then the volatile newest one from the device. 
-	 * @return					True if the hand was detected, false otherwise
+	 * Gets the position and orientation of the specified bone of the specified hand.
+	 * Data returned by this function is relative to the Unreal Engine origin rather than
+	 * a particular LeapMotionControllerComponent or LeapMotionHandActor instance.
+	 * Get the properties of a LeapMotionBoneActor instance to get the data
+	 * relative to that actor's parent hand actor and controller component. In this
+	 * case, use the standard Unreal Actor location and rotation properties.
+	 *
+	 * @param HandId			The id of the hand of interest.
+	 * @param LeapBone			The bone of interest.
+	 * @param OutPosition		An FVector set to the relative position of the hand, 
+	 *							or a zero vector if a hand with the specified id does not exist.
+	 * @param OutOrientation	An FRotator set to the relative orientation of 
+	 *							the specified bone, or a zero rotation if a hand 
+	 *							with the specified id does not exist.
+	 * @returns					True, if the hand with the specified id exists, false otherwise.
 	 */
 	UFUNCTION(BlueprintPure, Category = LeapMotion)
 	static bool GetBonePostionAndOrientation(int32 HandId, ELeapBone LeapBone, FVector& OutPosition, FRotator& OutOrientation);
 
 	/**
-	 * Gets the width and length of the specified bone of a hand from the device
-	 * @param HandId			Which hand to query
-	 * @param LeapBone			Which hand bone to get position & orientation for
-	 * @param OutWidth			Bone's width (in centimeters)
-	 * @param OutWidth			Bone's length (in centimeters)
-	 * @param UseReferenceFrame	Use saved reference frame, rather then the volatile newest one from the device.
-	 * @return					True if the hand was detected and bone was identified, false otherwise
+	 * Gets the width and length of the specified bone of the specified hand.
+	 * @param HandId			The id of the hand of interest.
+	 * @param LeapBone			The bone of interest.
+	 * @param OutWidth			A float set to the bone's width (in centimeters).
+	 * @param OutLength			A float set to the bone's length (in centimeters).
+	 * @returns					True, if the hand with the specified id exists, false otherwise.
 	 */
 	UFUNCTION(BlueprintPure, Category = LeapMotion)
 	static bool GetBoneWidthAndLength(int32 HandId, ELeapBone LeapBone, float& OutWidth, float& OutLength);
 
 	/**
-	 * Queries grabbing strength of the hand
-	 * @param HandId			Which hand to query
-	 * @param OutGrabStrength	Hand's grabbing strength as reported by Leap API
-	 * @param UseReferenceFrame	Use saved reference frame, rather then the volatile newest one from the device.
-	 * @return					True if the hand was detected, false otherwise
+	 * The grab strength rating of the specified hand.
+	 *
+	 * Grab strength is a rating of how the hand's posture resembles a fist.
+	 * A strength of 0 is close to an open, flat hand; a strength of 1 is close 
+	 * to a fist.
+	 * @param HandId			The id of the hand of interest.
+	 * @param GrabStrength		A float set to the hand's grabbing strength as reported by Leap API.
+	 * @returns					True, if the hand with the specified id exists, false otherwise.
 	 */
 	UFUNCTION(BlueprintPure, Category = LeapMotion)
 	static bool GetGrabStrength(int32 HandId, float& GrabStrength);
 
 	/**
-	 * Queries pinching strength of the hand
-	 * @param HandId			Which hand to query
-	 * @param OutPinchStrength	Hand's pinch strength as reported by Leap API
-	 * @param UseReferenceFrame	Use saved reference frame, rather then the volatile newest one from the device.
-	 * @return					True if the hand was detected, false otherwise
+	 * The pinch strength rating of the specified hand.
+	 *
+	 * Pinch strength is a rating of whether the hand is in a pinching posture.
+	 * Pinch strength starts at 0 and increases to 1 as any finger tip approaches
+	 * the tip of the thumb. Note that pinch and grab strength are not independent.
+	 * A grabbing hand will generally have a non-zero pinch strength as well.
+	 *
+	 * @param HandId			The id of the hand of interest.
+	 * @param PinchStrength		A float set to the hand's pinch strength as reported by Leap API
+	 * @returns					True, if the hand with the specified id exists, false otherwise.
 	 */
 	UFUNCTION(BlueprintPure, Category = LeapMotion)
 	static bool GetPinchStrength(int32 HandId, float& PinchStrength);
 
 	/**
-	* Set LeapController policy that improves tracking for HMD-mounted controller
-	* @param UseHmdPolicy		True for HMD-mounted mode, false otherwise
-	* @return					True if the device is connected
-	*/
-	UFUNCTION(BlueprintPure, Category = LeapMotion)
+	 * Enables or disables the Leap Motion Controller policy that improves tracking for an
+	 * HMD-mounted controller.
+	 * Note that calling this function does not change the transforms of any 
+	 * LeapMotionControllerComponent instances that may exist.
+	 *
+	 * @param UseHmdPolicy		True to enable for HMD-mounted mode, false to disable.
+	 * @returns					True if the device is connected.
+	 */
+	UFUNCTION(BlueprintCallable, Category = LeapMotion)
 	static bool SetHmdPolicy(bool UseHmdPolicy);
 };
 
diff --git a/Engine/Plugins/Runtime/LeapMotionController/Source/LeapMotionController/Public/LeapMotionTypes.h b/Engine/Plugins/Runtime/LeapMotionController/Source/LeapMotionController/Public/LeapMotionTypes.h
--- a/Engine/Plugins/Runtime/LeapMotionController/Source/LeapMotionController/Public/LeapMotionTypes.h
+++ b/Engine/Plugins/Runtime/LeapMotionController/Source/LeapMotionController/Public/LeapMotionTypes.h
@@ -6,14 +6,28 @@
 #include "Leap_NoPI.h"
 #include "LeapMotionTypes.generated.h"
 
+
+
 //---------------------------------------------------
 // Common enumerations
 //---------------------------------------------------
 
+/** 
+* \defgroup Enumerations Enumerations
+*
+* Leap Motion enumerations.
+* @{
+*/
+
+/**
+* The names for designating the hands.
+*/
 UENUM(BlueprintType)
 enum class ELeapSide : uint8
 {
+	/** Port */
 	Left,
+	/** Starboard */
 	Right
 };
 
@@ -21,25 +35,42 @@ enum class ELeapSide : uint8
 UENUM(BlueprintType)
 enum class ELeapBone : uint8
 {
+	/** The arm bone from elbow to wrist. */
 	Forearm,
+	/** The central portion of the hand, excluding fingers. */
 	Palm,
+	/** The anatomical thumb metacarpal. */
 	ThumbBase,
+	/** The anatomical thumb proximal phalanx. */
 	ThumbMiddle,
+	/** The anatomical thumb distal phalanx. */
 	ThumbTip,
+	/** The anatomical index finger proximal phalanx. */
 	Finger1Base,
+	/** The anatomical index finger intermediate phalanx. */
 	Finger1Middle,
+	/** The anatomical index finger distal phalanx. */
 	Finger1Tip,
+	/** The anatomical middle finger proximal phalanx. */
 	Finger2Base,
+	/** The anatomical middle finger intermediate phalanx. */
 	Finger2Middle,
+	/** The anatomical middle finger distal phalanx. */
 	Finger2Tip,
+	/** The anatomical ring finger proximal phalanx. */
 	Finger3Base,
+	/** The anatomical ring finger intermediate phalanx. */
 	Finger3Middle,
+	/** The anatomical ring finger distal phalanx. */
 	Finger3Tip,
+	/** The anatomical pinky finger proximal phalanx. */
 	Finger4Base,
+	/** The anatomical pink finger intermediate phalanx. */
 	Finger4Middle,
+	/** The anatomical pinky finger distal phalanx. */
 	Finger4Tip
 };
-
+/**@}*/ //end doxygen group
 
 //---------------------------------------------------
 // Delegate types
@@ -50,12 +81,15 @@ DECLARE_DYNAMIC_MULTICAST_DELEGATE_OneParam(FLMHandRemovedDelegate, int32, HandI
 DECLARE_DYNAMIC_MULTICAST_DELEGATE_TwoParams(FLMHandUpdatedDelegate, int32, HandId, float, DeltaSeconds);
 
 
-//---------------------------------------------------
-// Helper Functions
-//---------------------------------------------------
+/** ---------------------------------------------------
+* \defgroup global Utility Functions
+* Helper Functions
+* @{
+*/ //---------------------------------------------------
 
 /**
-* Helper function to convert from Leap space (in mm) to Unreal space (in cm)
+* Helper function to convert from Leap space to Unreal space (right-handed to left-handed)
+* and units (mm to cm).
 */
 static FORCEINLINE FVector LEAPVECTOR_TO_FVECTOR(Leap::Vector InVector)
 {
@@ -63,12 +97,13 @@ static FORCEINLINE FVector LEAPVECTOR_TO_FVECTOR(Leap::Vector InVector)
 }
 
 /**
-* Helper function to convert direction of the vector only
+* Helper function to convert direction of the vector only.
 */
 static FORCEINLINE FVector LEAPVECTOR_TO_FVECTOR_DIRECTION(Leap::Vector InVector)
 {
 	return FVector(-InVector.z, InVector.x, InVector.y);
 }
+/**@}*/ // End doxygen global group
 
 #define LEAP_BONE_NAME_CASE(LeapBone) case ELeapBone::LeapBone: return #LeapBone
 
