#leap Fix for crash when running Debug CRTs (UE-15444)

diff --git a/Engine/Plugins/Runtime/LeapMotionController/Source/LeapMotionController/LeapMotionController.Build.cs b/Engine/Plugins/Runtime/LeapMotionController/Source/LeapMotionController/LeapMotionController.Build.cs
--- a/Engine/Plugins/Runtime/LeapMotionController/Source/LeapMotionController/LeapMotionController.Build.cs
+++ b/Engine/Plugins/Runtime/LeapMotionController/Source/LeapMotionController/LeapMotionController.Build.cs
@@ -6,7 +6,15 @@ namespace UnrealBuildTool.Rules
 	{
 		public LeapMotionController(TargetInfo Target)
 		{
-			PublicDelayLoadDLLs.Add("Leap.dll");
+            if (Target.Configuration == UnrealTargetConfiguration.Debug && BuildConfiguration.bDebugBuildsActuallyUseDebugCRT)
+            {
+                Definitions.Add("LEAP_USE_DEBUG_LIB=1");
+                PublicDelayLoadDLLs.Add("Leapd.dll");
+            }
+            else
+            {
+                PublicDelayLoadDLLs.Add("Leap.dll");
+            }
 
 			PublicDependencyModuleNames.AddRange( new string[] { "Core", "CoreUObject", "Engine", "InputCore" } ); 
 
diff --git a/Engine/Plugins/Runtime/LeapMotionController/Source/LeapMotionController/Private/LeapMotionControllerPlugin.cpp b/Engine/Plugins/Runtime/LeapMotionController/Source/LeapMotionController/Private/LeapMotionControllerPlugin.cpp
--- a/Engine/Plugins/Runtime/LeapMotionController/Source/LeapMotionController/Private/LeapMotionControllerPlugin.cpp
+++ b/Engine/Plugins/Runtime/LeapMotionController/Source/LeapMotionController/Private/LeapMotionControllerPlugin.cpp
@@ -17,7 +17,11 @@ void FLeapMotionControllerPlugin::StartupModule()
 {
 	FString LeapDllRoot = FPaths::EngineDir() / TEXT("Binaries/ThirdParty/Leap/") / (PLATFORM_64BITS ? TEXT("Win64/") : TEXT("Win32/"));
 	FPlatformProcess::PushDllDirectory(*LeapDllRoot);
+#if LEAP_USE_DEBUG_LIB
+  	LeapDllHandle = FPlatformProcess::GetDllHandle(*(LeapDllRoot + "Leapd.dll"));
+#else
 	LeapDllHandle = FPlatformProcess::GetDllHandle(*(LeapDllRoot + "Leap.dll"));
+#endif
 	FPlatformProcess::PopDllDirectory(*LeapDllRoot);
 
 	// Continue with initialization if dll loads successfully
