[INTEGRATE] Change 2571721 by [REDACTED_PERSON][REDACTED_HANDLE] on 2015/05/30 12:50:57

	Fix for Leap causing issues with windows xp builds.

	Leap.dll is not compatible, so we need to dont even attempt to load it on xp supported builds.

diff --git a/Engine/Plugins/Runtime/LeapMotionController/Source/LeapMotionController/Private/LeapMotionControllerPlugin.cpp b/Engine/Plugins/Runtime/LeapMotionController/Source/LeapMotionController/Private/LeapMotionControllerPlugin.cpp
--- a/Engine/Plugins/Runtime/LeapMotionController/Source/LeapMotionController/Private/LeapMotionControllerPlugin.cpp
+++ b/Engine/Plugins/Runtime/LeapMotionController/Source/LeapMotionController/Private/LeapMotionControllerPlugin.cpp
@@ -15,6 +15,10 @@ static void* LeapDllHandle = 0;
 
 void FLeapMotionControllerPlugin::StartupModule()
 {
+#if PLATFORM_WINDOWS && ( WINVER == 0x502 )	// Leap not supported on Windows XP
+	return;
+#endif
+
 	FString LeapDllRoot = FPaths::EngineDir() / TEXT("Binaries/ThirdParty/Leap/") / (PLATFORM_64BITS ? TEXT("Win64/") : TEXT("Win32/"));
 	FPlatformProcess::PushDllDirectory(*LeapDllRoot);
 #if LEAP_USE_DEBUG_LIB
