#leapmotion Fix for jittery movement.  BoneActor updates are now swept updates, instead of simulated physics movements

diff --git a/Engine/Plugins/Runtime/LeapMotionController/Source/LeapMotionController/Private/LeapMotionBoneActor.cpp b/Engine/Plugins/Runtime/LeapMotionController/Source/LeapMotionController/Private/LeapMotionBoneActor.cpp
--- a/Engine/Plugins/Runtime/LeapMotionController/Source/LeapMotionController/Private/LeapMotionBoneActor.cpp
+++ b/Engine/Plugins/Runtime/LeapMotionController/Source/LeapMotionController/Private/LeapMotionBoneActor.cpp
@@ -16,7 +16,6 @@ ALeapMotionBoneActor::ALeapMotionBoneActor(const class FObjectInitializer& Objec
 
 	// Initialize root primitive component. It only serves as a reference point. Collision & visual meshes are offset from it.
 	USphereComponent* SphereComponent = ObjectInitializer.CreateDefaultSubobject<USphereComponent>(this, TEXT("RootSphereComponent"));
-	SphereComponent->SetSimulatePhysics(true);
 	SphereComponent->SetSphereRadius(0.001);
 	SphereComponent->SetPhysicsMaxAngularVelocity(7200.0f);
 	SphereComponent->SetCollisionProfileName(UCollisionProfile::CustomCollisionProfileName);
diff --git a/Engine/Plugins/Runtime/LeapMotionController/Source/LeapMotionController/Private/LeapMotionHandActor.cpp b/Engine/Plugins/Runtime/LeapMotionController/Source/LeapMotionController/Private/LeapMotionHandActor.cpp
--- a/Engine/Plugins/Runtime/LeapMotionController/Source/LeapMotionController/Private/LeapMotionHandActor.cpp
+++ b/Engine/Plugins/Runtime/LeapMotionController/Source/LeapMotionController/Private/LeapMotionHandActor.cpp
@@ -166,7 +166,6 @@ void ALeapMotionHandActor::UpdateBones(float DeltaSeconds)
 			FRotator TargetOrientation;
 
 			bool Success = Device->GetBonePostionAndOrientation(HandId, LeapBone, TargetPosition, TargetOrientation);
-
 			if (Success)
 			{
 				// Offset target position & rotation by the SpawnReference actor's transform
@@ -176,26 +175,10 @@ void ALeapMotionHandActor::UpdateBones(float DeltaSeconds)
 
 				// Get current position & rotation
 				ALeapMotionBoneActor* BoneActor = BoneActors[BoneArrayIndex++];
-
 				UPrimitiveComponent* PrimitiveComponent = Cast<UPrimitiveComponent>(BoneActor->GetRootComponent());
-				if (PrimitiveComponent && PrimitiveComponent->IsSimulatingPhysics())
+				if (PrimitiveComponent)
 				{
-					FVector CurrentPositon = PrimitiveComponent->GetComponentLocation();
-					FRotator CurrentRotation = PrimitiveComponent->GetComponentRotation();
-
-					// Compute linear velocity
-					FVector LinearVelocity = (TargetPosition - CurrentPositon) / DeltaSeconds;
-
-					// Compute angular velocity
-					FVector Axis;
-					float Angle;
-					ConvertDeltaRotationsToAxisAngle(CurrentRotation, TargetOrientation, Axis, Angle);
-					if (Angle > PI) { Angle -= 2 * PI; }
-					FVector AngularVelcity = Axis * (Angle / DeltaSeconds);
-
-					// Apply velocities
-					PrimitiveComponent->SetPhysicsLinearVelocity(LinearVelocity);
-					PrimitiveComponent->SetAllPhysicsAngularVelocity(AngularVelcity * 180.0f / PI);
+					PrimitiveComponent->SetRelativeLocationAndRotation(TargetPosition, TargetOrientation, true);
 				}
 			}
 		}
